Added string conversion for GeForceNow OsTypes

diff --git a/Engine/Source/Runtime/NVidia/GeForceNOW/Private/GeForceNOWWrapper.cpp b/Engine/Source/Runtime/NVidia/GeForceNOW/Private/GeForceNOWWrapper.cpp
--- a/Engine/Source/Runtime/NVidia/GeForceNOW/Private/GeForceNOWWrapper.cpp
+++ b/Engine/Source/Runtime/NVidia/GeForceNOW/Private/GeForceNOWWrapper.cpp
@@ -34,6 +34,40 @@ GeForceNOWWrapper& GeForceNOWWrapper::Get()
 	return Singleton;
 }
 
+/*static*/ const FString GeForceNOWWrapper::GetGfnOsTypeString(GfnOsType OsType)
+{
+	switch (OsType)
+	{
+	case gfnWindows:
+		return "gfnWindows";
+	case gfnMacOs:
+		return "gfnMacOs";
+	case gfnShield:
+		return "gfnShield";
+	case gfnAndroid:
+		return "gfnAndroid";
+	case gfnIOs:
+		return "gfnIOs";
+	case gfnIPadOs:
+		return "gfnIPadOs";
+	case gfnChromeOs:
+		return "gfnChromeOs";
+	case gfnLinux:
+		return "gfnLinux";
+	case gfnTizen:
+		return "gfnTizen";
+	case gfnWebOs:
+		return "gfnWebOs";
+	case gfnTvOs:
+		return "gfnTvOs";
+	case gfnUnknownOs:
+	default:
+		return "gfnUnknownOs";
+	}
+	static_assert(GfnOsType::gfnOsTypeMax == 11, "Error: Please update GetGfnOSTypeString with new GfnOSType entries.");
+}
+
+
 bool GeForceNOWWrapper::IsRunningInGFN()
 {
 	if (IsRunningMockGFN())
diff --git a/Engine/Source/Runtime/NVidia/GeForceNOW/Public/GeForceNOWWrapper.h b/Engine/Source/Runtime/NVidia/GeForceNOW/Public/GeForceNOWWrapper.h
--- a/Engine/Source/Runtime/NVidia/GeForceNOW/Public/GeForceNOWWrapper.h
+++ b/Engine/Source/Runtime/NVidia/GeForceNOW/Public/GeForceNOWWrapper.h
@@ -21,6 +21,8 @@ public:
 	virtual ~GeForceNOWWrapper();
 	static GeForceNOWWrapper& Get();
 
+	static const FString GetGfnOsTypeString(GfnOsType OsType);
+
 	/** Returns true if the GeForceNow SDK is initialized and running in cloud. */
 	bool IsRunningInGFN();
 
